Making plugin not enabled by default and setting to experimental

diff --git a/Engine/Plugins/Experimental/MotoSynth/MotoSynth.uplugin b/Engine/Plugins/Experimental/MotoSynth/MotoSynth.uplugin
--- a/Engine/Plugins/Experimental/MotoSynth/MotoSynth.uplugin
+++ b/Engine/Plugins/Experimental/MotoSynth/MotoSynth.uplugin
@@ -3,16 +3,16 @@
 	"Version" : 1,
 	"VersionName" : "1",
 	"FriendlyName" : "MotoSynth",
-	"Description" : "An experimental granular vehicle engine.",
+	"Description" : "An experimental granular vehicle engine. Intended to explore and demonstrate potential capabilities. Not supported.",
 	"Category" : "Audio",
 	"CreatedBy" : "Epic Games, Inc.",
 	"CreatedByURL" : "http://epicgames.com",
 	"DocsURL" : "",
 	"MarketplaceURL" : "",
 	"SupportURL" : "",
-	"EnabledByDefault" : true,
+	"EnabledByDefault" : false,
 	"CanContainContent" : true,
-	"IsBetaVersion" : true,
+	"IsExperimentalVersion" : true,
 	"Installed" : false,
 	"Modules" :
 	[
