Audio Plugin Telemetry Part 2
AudioGameplayVolumes, MicrosoftSpatialSound, MotoSynth, and OculusAudio

diff --git a/Engine/Plugins/AudioGameplayVolume/Source/AudioGameplayVolume/Private/AudioGameplayVolumeProxy.cpp b/Engine/Plugins/AudioGameplayVolume/Source/AudioGameplayVolume/Private/AudioGameplayVolumeProxy.cpp
--- a/Engine/Plugins/AudioGameplayVolume/Source/AudioGameplayVolume/Private/AudioGameplayVolumeProxy.cpp
+++ b/Engine/Plugins/AudioGameplayVolume/Source/AudioGameplayVolume/Private/AudioGameplayVolumeProxy.cpp
@@ -5,6 +5,7 @@
 #include "AudioGameplayVolumeSubsystem.h"
 #include "AudioGameplayVolumeLogs.h"
 #include "AudioGameplayVolumeComponent.h"
+#include "EngineAnalytics.h"
 #include "Interfaces/IAudioGameplayCondition.h"
 #include "Components/BrushComponent.h"
 #include "Components/PrimitiveComponent.h"
@@ -51,6 +52,12 @@ void UAudioGameplayVolumeProxy::InitFromComponent(const UAudioGameplayVolumeProx
 			ProxyVolumeMutators.Emplace(NewMutator);
 		}
 	}
+
+	if (FEngineAnalytics::IsAvailable())
+	{
+		FEngineAnalytics::GetProvider().RecordEvent(TEXT("Audio.Usage.AudioGameplayVolume.InitializedFromComponent"));
+	}
+
 }
 
 void UAudioGameplayVolumeProxy::FindMutatorPriority(FAudioProxyMutatorPriorities& Priorities) const
diff --git a/Engine/Plugins/Experimental/MotoSynth/Source/MotoSynthEditor/Private/MotoSynthSourceFactory.cpp b/Engine/Plugins/Experimental/MotoSynth/Source/MotoSynthEditor/Private/MotoSynthSourceFactory.cpp
--- a/Engine/Plugins/Experimental/MotoSynth/Source/MotoSynthEditor/Private/MotoSynthSourceFactory.cpp
+++ b/Engine/Plugins/Experimental/MotoSynth/Source/MotoSynthEditor/Private/MotoSynthSourceFactory.cpp
@@ -1,9 +1,10 @@
 // Copyright Epic Games, Inc. All Rights Reserved.
 #include "MotoSynthSourceFactory.h"
 #include "MotoSynthSourceAsset.h"
+#include "EngineAnalytics.h"
+#include "MotoSynthEditorModule.h"
 #include "UObject/ObjectMacros.h"
 #include "UObject/Object.h"
-#include "MotoSynthEditorModule.h"
 #include "Sound/SoundWave.h"
 
 #define LOCTEXT_NAMESPACE "AssetTypeActions"
@@ -36,6 +37,10 @@ UMotoSynthPresetFactory::UMotoSynthPresetFactory(const FObjectInitializer& Objec
 UObject* UMotoSynthPresetFactory::FactoryCreateNew(UClass* Class, UObject* InParent, FName InName, EObjectFlags Flags, UObject* Context, FFeedbackContext* Warn)
 {
 	UMotoSynthPreset* NewAsset = NewObject<UMotoSynthPreset >(InParent, InName, Flags);
+	if (FEngineAnalytics::IsAvailable())
+	{
+		FEngineAnalytics::GetProvider().RecordEvent(TEXT("Audio.Usage.MotoSynth.PresetCreated"));
+	}
 	return NewAsset;
 }
 
@@ -74,6 +79,12 @@ UObject* UMotoSynthSourceFactory::FactoryCreateNew(UClass* Class, UObject* InPar
 			NewAsset->SoundWaveSource = SoundWave;
 			NewAsset->UpdateSourceData();
 			NewAsset->PerformGrainTableAnalysis();
+
+			if (FEngineAnalytics::IsAvailable())
+			{
+				FEngineAnalytics::GetProvider().RecordEvent(TEXT("Audio.Usage.MotoSynth.SourceCreated"));
+			}
+
 			return NewAsset;
 		}
 
diff --git a/Engine/Plugins/Runtime/MicrosoftSpatialAudio/Source/MicrosoftSpatialSound/Private/MicrosoftSpatialSoundPlugin.cpp b/Engine/Plugins/Runtime/MicrosoftSpatialAudio/Source/MicrosoftSpatialSound/Private/MicrosoftSpatialSoundPlugin.cpp
--- a/Engine/Plugins/Runtime/MicrosoftSpatialAudio/Source/MicrosoftSpatialSound/Private/MicrosoftSpatialSoundPlugin.cpp
+++ b/Engine/Plugins/Runtime/MicrosoftSpatialAudio/Source/MicrosoftSpatialSound/Private/MicrosoftSpatialSoundPlugin.cpp
@@ -1,6 +1,7 @@
 // Copyright Epic Games, Inc. All Rights Reserved.
 
 #include "MicrosoftSpatialSoundPlugin.h"
+#include "EngineAnalytics.h"
 #include "Features/IModularFeatures.h"
 #include "HAL/PlatformMisc.h"
 
@@ -211,6 +212,11 @@ void FMicrosoftSpatialSound::Initialize(const FAudioPluginInitializationParams I
 			bWarnedMicrosoftSpatialSoundDynamicObjectCountIsZero = false;
 
 			SpatialAudioRenderThread = FRunnableThread::Create(this, TEXT("MicrosoftSpatialAudioThread"), 0, TPri_TimeCritical, FPlatformAffinity::GetAudioThreadMask());
+
+			if (FEngineAnalytics::IsAvailable())
+			{
+				FEngineAnalytics::GetProvider().RecordEvent(TEXT("Audio.Usage.MicrosoftSpatialSound.Initialized"));
+			}
 		}
 		else
 		{
diff --git a/Engine/Plugins/Runtime/Oculus/OculusAudio/Source/OculusAudioEditor/Private/OculusAmbisonicSettingsFactory.cpp b/Engine/Plugins/Runtime/Oculus/OculusAudio/Source/OculusAudioEditor/Private/OculusAmbisonicSettingsFactory.cpp
--- a/Engine/Plugins/Runtime/Oculus/OculusAudio/Source/OculusAudioEditor/Private/OculusAmbisonicSettingsFactory.cpp
+++ b/Engine/Plugins/Runtime/Oculus/OculusAudio/Source/OculusAudioEditor/Private/OculusAmbisonicSettingsFactory.cpp
@@ -3,6 +3,7 @@
 #include "OculusAmbisonicSettingsFactory.h"
 #include "OculusAmbisonicsSettings.h"
 #include "AssetTypeCategories.h"
+#include "EngineAnalytics.h"
 #include "UObject/ObjectMacros.h"
 #include "UObject/Object.h"
 
@@ -51,6 +52,10 @@ UOculusAmbisonicsSettingsFactory::UOculusAmbisonicsSettingsFactory(const FObject
 UObject* UOculusAmbisonicsSettingsFactory::FactoryCreateNew(UClass* Class, UObject* InParent, FName Name, EObjectFlags Flags, UObject* Context,
 	FFeedbackContext* Warn)
 {
+	if (FEngineAnalytics::IsAvailable())
+	{
+		FEngineAnalytics::GetProvider().RecordEvent(TEXT("Audio.Usage.OculusAudio.AmbisonicsSettingsCreated"));
+	}
 	return NewObject<UOculusAudioSoundfieldSettings>(InParent, Name, Flags);
 }
 
diff --git a/Engine/Plugins/Runtime/Oculus/OculusAudio/Source/OculusAudioEditor/Private/OculusAudioSourceSettingsFactory.cpp b/Engine/Plugins/Runtime/Oculus/OculusAudio/Source/OculusAudioEditor/Private/OculusAudioSourceSettingsFactory.cpp
--- a/Engine/Plugins/Runtime/Oculus/OculusAudio/Source/OculusAudioEditor/Private/OculusAudioSourceSettingsFactory.cpp
+++ b/Engine/Plugins/Runtime/Oculus/OculusAudio/Source/OculusAudioEditor/Private/OculusAudioSourceSettingsFactory.cpp
@@ -3,6 +3,7 @@
 #include "OculusAudioSourceSettingsFactory.h"
 #include "OculusAudioSourceSettings.h"
 #include "AssetTypeCategories.h"
+#include "EngineAnalytics.h"
 #include "UObject/ObjectMacros.h"
 #include "UObject/Object.h"
 
@@ -51,6 +52,10 @@ UOculusAudioSourceSettingsFactory::UOculusAudioSourceSettingsFactory(const FObje
 UObject* UOculusAudioSourceSettingsFactory::FactoryCreateNew(UClass* Class, UObject* InParent, FName Name, EObjectFlags Flags, UObject* Context,
     FFeedbackContext* Warn)
 {
+	if (FEngineAnalytics::IsAvailable())
+	{
+		FEngineAnalytics::GetProvider().RecordEvent(TEXT("Audio.Usage.OculusAudio.SourceSettingsCreated"));
+	}
     return NewObject<UOculusAudioSourceSettings>(InParent, Name, Flags);
 }
 
