Fix non-unity issues
#rnx

diff --git a/Engine/Plugins/Experimental/MotoSynth/Source/MotoSynth/Private/MotoSynthDataManager.cpp b/Engine/Plugins/Experimental/MotoSynth/Source/MotoSynth/Private/MotoSynthDataManager.cpp
--- a/Engine/Plugins/Experimental/MotoSynth/Source/MotoSynth/Private/MotoSynthDataManager.cpp
+++ b/Engine/Plugins/Experimental/MotoSynth/Source/MotoSynth/Private/MotoSynthDataManager.cpp
@@ -3,6 +3,8 @@
 #include "MotoSynthDataManager.h"
 #include "MotoSynthModule.h"
 
+#include "HAL/IConsoleManager.h"
+
 
 // Log all loaded data in moto synth data manager
 static FAutoConsoleCommand LogMotoSynthMemoryUsageCommand(
diff --git a/Engine/Plugins/Media/AjaMedia/Source/AjaMediaOutput/Private/AjaMediaCapture.cpp b/Engine/Plugins/Media/AjaMedia/Source/AjaMediaOutput/Private/AjaMediaCapture.cpp
--- a/Engine/Plugins/Media/AjaMedia/Source/AjaMediaOutput/Private/AjaMediaCapture.cpp
+++ b/Engine/Plugins/Media/AjaMedia/Source/AjaMediaOutput/Private/AjaMediaCapture.cpp
@@ -6,6 +6,7 @@
 #include "AjaDeviceProvider.h"
 #include "AjaMediaOutput.h"
 #include "AjaMediaOutputModule.h"
+#include "Engine/Engine.h"
 #include "HAL/Event.h"
 #include "HAL/IConsoleManager.h"
 #include "IAjaMediaOutputModule.h"
