CIS Followup - deleting pointer to incomplete type.

diff --git a/Engine/Plugins/Experimental/MotoSynth/Source/MotoSynth/Private/MotoSynthPreset.cpp b/Engine/Plugins/Experimental/MotoSynth/Source/MotoSynth/Private/MotoSynthPreset.cpp
--- a/Engine/Plugins/Experimental/MotoSynth/Source/MotoSynth/Private/MotoSynthPreset.cpp
+++ b/Engine/Plugins/Experimental/MotoSynth/Source/MotoSynth/Private/MotoSynthPreset.cpp
@@ -106,6 +106,10 @@ FMotoSynthEnginePreviewer::FMotoSynthEnginePreviewer()
 	SynthEngine = TUniquePtr<FMotoSynthEngine>(new FMotoSynthEngine);
 }
 
+FMotoSynthEnginePreviewer::~FMotoSynthEnginePreviewer()
+{
+}
+
 void FMotoSynthEnginePreviewer::SetSettings(const FMotoSynthRuntimeSettings& InSettings)
 {
 	FScopeLock Lock(&PreviewEngineCritSect);
diff --git a/Engine/Plugins/Experimental/MotoSynth/Source/MotoSynth/Public/MotoSynthPreset.h b/Engine/Plugins/Experimental/MotoSynth/Source/MotoSynth/Public/MotoSynthPreset.h
--- a/Engine/Plugins/Experimental/MotoSynth/Source/MotoSynth/Public/MotoSynthPreset.h
+++ b/Engine/Plugins/Experimental/MotoSynth/Source/MotoSynth/Public/MotoSynthPreset.h
@@ -175,7 +175,7 @@ class FMotoSynthEnginePreviewer : public ISubmixBufferListener
 {
 public:
 	FMotoSynthEnginePreviewer();
-	virtual ~FMotoSynthEnginePreviewer() override = default;
+	virtual ~FMotoSynthEnginePreviewer() override;
 
 	void StartPreviewing();
 	void StopPreviewing();
