fixup ISoundGenerator::GetNextBuffer() implementers (don't assume zero'd buffer)
#peflight 68af91a64c47060e8499b1e6

diff --git a/Engine/Plugins/Experimental/MotoSynth/Source/MotoSynth/Private/MotoSynthEngine.cpp b/Engine/Plugins/Experimental/MotoSynth/Source/MotoSynth/Private/MotoSynthEngine.cpp
--- a/Engine/Plugins/Experimental/MotoSynth/Source/MotoSynth/Private/MotoSynthEngine.cpp
+++ b/Engine/Plugins/Experimental/MotoSynth/Source/MotoSynth/Private/MotoSynthEngine.cpp
@@ -571,6 +571,8 @@ int32 FMotoSynthEngine::OnGenerateAudio(float* OutAudio, int32 NumSamples)
 	// Don't do anything if we're not enabled
 	if (!IsMotoSynthEngineEnabled())
 	{
+		FMemory::Memzero(OutAudio, sizeof(float) * NumSamples);
+		
 		return NumSamples;
 	}
 
